Fix copyrights for new soundfield/endpoint files.

diff --git a/Engine/Plugins/Runtime/SoundFields/Source/SoundFields/Private/SoundFields.cpp b/Engine/Plugins/Runtime/SoundFields/Source/SoundFields/Private/SoundFields.cpp
--- a/Engine/Plugins/Runtime/SoundFields/Source/SoundFields/Private/SoundFields.cpp
+++ b/Engine/Plugins/Runtime/SoundFields/Source/SoundFields/Private/SoundFields.cpp
@@ -1,4 +1,4 @@
-// Copyright 1998-2019 Epic Games, Inc. All Rights Reserved.
+// Copyright Epic Games, Inc. All Rights Reserved.
 
 #include "SoundFields.h"
 #include "CoreMinimal.h"
diff --git a/Engine/Plugins/Runtime/SoundFields/Source/SoundFields/Private/SoundFieldsModule.cpp b/Engine/Plugins/Runtime/SoundFields/Source/SoundFields/Private/SoundFieldsModule.cpp
--- a/Engine/Plugins/Runtime/SoundFields/Source/SoundFields/Private/SoundFieldsModule.cpp
+++ b/Engine/Plugins/Runtime/SoundFields/Source/SoundFields/Private/SoundFieldsModule.cpp
@@ -1,4 +1,4 @@
-// Copyright 1998-2019 Epic Games, Inc. All Rights Reserved.
+// Copyright Epic Games, Inc. All Rights Reserved.
 
 #include "SoundFieldsModule.h"
 #include "Features/IModularFeatures.h"
diff --git a/Engine/Plugins/Runtime/SoundFields/Source/SoundFields/Public/SoundFields.h b/Engine/Plugins/Runtime/SoundFields/Source/SoundFields/Public/SoundFields.h
--- a/Engine/Plugins/Runtime/SoundFields/Source/SoundFields/Public/SoundFields.h
+++ b/Engine/Plugins/Runtime/SoundFields/Source/SoundFields/Public/SoundFields.h
@@ -1,4 +1,4 @@
-// Copyright 1998-2019 Epic Games, Inc. All Rights Reserved.
+// Copyright Epic Games, Inc. All Rights Reserved.
 
 #pragma once
 
diff --git a/Engine/Plugins/Runtime/SoundFields/Source/SoundFields/Public/SoundFieldsModule.h b/Engine/Plugins/Runtime/SoundFields/Source/SoundFields/Public/SoundFieldsModule.h
--- a/Engine/Plugins/Runtime/SoundFields/Source/SoundFields/Public/SoundFieldsModule.h
+++ b/Engine/Plugins/Runtime/SoundFields/Source/SoundFields/Public/SoundFieldsModule.h
@@ -1,4 +1,4 @@
-// Copyright 1998-2019 Epic Games, Inc. All Rights Reserved.
+// Copyright Epic Games, Inc. All Rights Reserved.
 
 #pragma once
 
diff --git a/Engine/Plugins/Runtime/SoundFields/Source/SoundFields/SoundFields.Build.cs b/Engine/Plugins/Runtime/SoundFields/Source/SoundFields/SoundFields.Build.cs
--- a/Engine/Plugins/Runtime/SoundFields/Source/SoundFields/SoundFields.Build.cs
+++ b/Engine/Plugins/Runtime/SoundFields/Source/SoundFields/SoundFields.Build.cs
@@ -1,4 +1,4 @@
-// Copyright 1998-2019 Epic Games, Inc. All Rights Reserved.
+// Copyright Epic Games, Inc. All Rights Reserved.
 
 using UnrealBuildTool;
 using System;
diff --git a/Engine/Source/Editor/AudioEditor/Private/SubmixDetailsCustomization.cpp b/Engine/Source/Editor/AudioEditor/Private/SubmixDetailsCustomization.cpp
--- a/Engine/Source/Editor/AudioEditor/Private/SubmixDetailsCustomization.cpp
+++ b/Engine/Source/Editor/AudioEditor/Private/SubmixDetailsCustomization.cpp
@@ -1,4 +1,4 @@
-// Copyright 1998-2019 Epic Games, Inc. All Rights Reserved.
+// Copyright Epic Games, Inc. All Rights Reserved.
 #include "SubmixDetailsCustomization.h"
 #include "Misc/PackageName.h"
 #include "Widgets/Layout/SBox.h"
diff --git a/Engine/Source/Editor/AudioEditor/Public/SubmixDetailsCustomization.h b/Engine/Source/Editor/AudioEditor/Public/SubmixDetailsCustomization.h
--- a/Engine/Source/Editor/AudioEditor/Public/SubmixDetailsCustomization.h
+++ b/Engine/Source/Editor/AudioEditor/Public/SubmixDetailsCustomization.h
@@ -1,4 +1,4 @@
-// Copyright 1998-2019 Epic Games, Inc. All Rights Reserved.
+// Copyright Epic Games, Inc. All Rights Reserved.
 #pragma once
 
 #include "CoreMinimal.h"
diff --git a/Engine/Source/Runtime/AudioExtensions/AudioExtensions.Build.cs b/Engine/Source/Runtime/AudioExtensions/AudioExtensions.Build.cs
--- a/Engine/Source/Runtime/AudioExtensions/AudioExtensions.Build.cs
+++ b/Engine/Source/Runtime/AudioExtensions/AudioExtensions.Build.cs
@@ -1,4 +1,4 @@
-// Copyright 1998-2019 Epic Games, Inc. All Rights Reserved.
+// Copyright Epic Games, Inc. All Rights Reserved.
 
 namespace UnrealBuildTool.Rules
 {
diff --git a/Engine/Source/Runtime/AudioExtensions/Private/IAudioEndpoint.cpp b/Engine/Source/Runtime/AudioExtensions/Private/IAudioEndpoint.cpp
--- a/Engine/Source/Runtime/AudioExtensions/Private/IAudioEndpoint.cpp
+++ b/Engine/Source/Runtime/AudioExtensions/Private/IAudioEndpoint.cpp
@@ -1,4 +1,4 @@
-// Copyright 1998-2019 Epic Games, Inc. All Rights Reserved.
+// Copyright Epic Games, Inc. All Rights Reserved.
 
 #include "IAudioEndpoint.h"
 
diff --git a/Engine/Source/Runtime/AudioExtensions/Private/ISoundfieldEndpoint.cpp b/Engine/Source/Runtime/AudioExtensions/Private/ISoundfieldEndpoint.cpp
--- a/Engine/Source/Runtime/AudioExtensions/Private/ISoundfieldEndpoint.cpp
+++ b/Engine/Source/Runtime/AudioExtensions/Private/ISoundfieldEndpoint.cpp
@@ -1,4 +1,4 @@
-// Copyright 1998-2019 Epic Games, Inc. All Rights Reserved.
+// Copyright Epic Games, Inc. All Rights Reserved.
 
 #include "ISoundfieldEndpoint.h"
 
diff --git a/Engine/Source/Runtime/AudioExtensions/Public/IAudioEndpoint.h b/Engine/Source/Runtime/AudioExtensions/Public/IAudioEndpoint.h
--- a/Engine/Source/Runtime/AudioExtensions/Public/IAudioEndpoint.h
+++ b/Engine/Source/Runtime/AudioExtensions/Public/IAudioEndpoint.h
@@ -1,4 +1,4 @@
-// Copyright 1998-2019 Epic Games, Inc. All Rights Reserved.
+// Copyright Epic Games, Inc. All Rights Reserved.
 
 #pragma once
 
diff --git a/Engine/Source/Runtime/AudioExtensions/Public/ISoundfieldEndpoint.h b/Engine/Source/Runtime/AudioExtensions/Public/ISoundfieldEndpoint.h
--- a/Engine/Source/Runtime/AudioExtensions/Public/ISoundfieldEndpoint.h
+++ b/Engine/Source/Runtime/AudioExtensions/Public/ISoundfieldEndpoint.h
@@ -1,4 +1,4 @@
-// Copyright 1998-2019 Epic Games, Inc. All Rights Reserved.
+// Copyright Epic Games, Inc. All Rights Reserved.
 
 #pragma once
 
diff --git a/Engine/Source/Runtime/SoundFieldRendering/Private/SoundFieldRenderingModule.cpp b/Engine/Source/Runtime/SoundFieldRendering/Private/SoundFieldRenderingModule.cpp
--- a/Engine/Source/Runtime/SoundFieldRendering/Private/SoundFieldRenderingModule.cpp
+++ b/Engine/Source/Runtime/SoundFieldRendering/Private/SoundFieldRenderingModule.cpp
@@ -1,4 +1,4 @@
-// Copyright 1998-2019 Epic Games, Inc. All Rights Reserved.
+// Copyright Epic Games, Inc. All Rights Reserved.
 
 #include "SoundFieldRenderingModule.h"
 #include "Features/IModularFeatures.h"
diff --git a/Engine/Source/Runtime/SoundFieldRendering/Private/SphericalHarmonicCalculator.cpp b/Engine/Source/Runtime/SoundFieldRendering/Private/SphericalHarmonicCalculator.cpp
--- a/Engine/Source/Runtime/SoundFieldRendering/Private/SphericalHarmonicCalculator.cpp
+++ b/Engine/Source/Runtime/SoundFieldRendering/Private/SphericalHarmonicCalculator.cpp
@@ -1,4 +1,4 @@
-// Copyright 1998-2019 Epic Games, Inc. All Rights Reserved.
+// Copyright Epic Games, Inc. All Rights Reserved.
 
 #pragma once
 #include "SphericalHarmonicCalculator.h"
diff --git a/Engine/Source/Runtime/SoundFieldRendering/Public/SoundFieldRendering.h b/Engine/Source/Runtime/SoundFieldRendering/Public/SoundFieldRendering.h
--- a/Engine/Source/Runtime/SoundFieldRendering/Public/SoundFieldRendering.h
+++ b/Engine/Source/Runtime/SoundFieldRendering/Public/SoundFieldRendering.h
@@ -1,4 +1,4 @@
-// Copyright 1998-2019 Epic Games, Inc. All Rights Reserved.
+// Copyright Epic Games, Inc. All Rights Reserved.
 
 #pragma once
 
diff --git a/Engine/Source/Runtime/SoundFieldRendering/Public/SoundFieldRenderingModule.h b/Engine/Source/Runtime/SoundFieldRendering/Public/SoundFieldRenderingModule.h
--- a/Engine/Source/Runtime/SoundFieldRendering/Public/SoundFieldRenderingModule.h
+++ b/Engine/Source/Runtime/SoundFieldRendering/Public/SoundFieldRenderingModule.h
@@ -1,4 +1,4 @@
-// Copyright 1998-2019 Epic Games, Inc. All Rights Reserved.
+// Copyright Epic Games, Inc. All Rights Reserved.
 
 #pragma once
 
diff --git a/Engine/Source/Runtime/SoundFieldRendering/Public/SphericalHarmonicCalculator.h b/Engine/Source/Runtime/SoundFieldRendering/Public/SphericalHarmonicCalculator.h
--- a/Engine/Source/Runtime/SoundFieldRendering/Public/SphericalHarmonicCalculator.h
+++ b/Engine/Source/Runtime/SoundFieldRendering/Public/SphericalHarmonicCalculator.h
@@ -1,4 +1,4 @@
-// Copyright 1998-2019 Epic Games, Inc. All Rights Reserved.
+// Copyright Epic Games, Inc. All Rights Reserved.
 
 #pragma once
 
diff --git a/Engine/Source/Runtime/SoundFieldRendering/SoundFieldRendering.Build.cs b/Engine/Source/Runtime/SoundFieldRendering/SoundFieldRendering.Build.cs
--- a/Engine/Source/Runtime/SoundFieldRendering/SoundFieldRendering.Build.cs
+++ b/Engine/Source/Runtime/SoundFieldRendering/SoundFieldRendering.Build.cs
@@ -1,4 +1,4 @@
-// Copyright 1998-2019 Epic Games, Inc. All Rights Reserved.
+// Copyright Epic Games, Inc. All Rights Reserved.
 
 namespace UnrealBuildTool.Rules
 {
