Removed "OptimizeCode = CodeOptimization.Never" from SoundFields. This setting created a 2.03 GB shared PCH file that was only used by this module.

diff --git a/Engine/Plugins/Runtime/SoundFields/Source/SoundFields/SoundFields.Build.cs b/Engine/Plugins/Runtime/SoundFields/Source/SoundFields/SoundFields.Build.cs
--- a/Engine/Plugins/Runtime/SoundFields/Source/SoundFields/SoundFields.Build.cs
+++ b/Engine/Plugins/Runtime/SoundFields/Source/SoundFields/SoundFields.Build.cs
@@ -8,9 +8,8 @@ public class SoundFields : ModuleRules
 	public SoundFields(ReadOnlyTargetRules Target) : base(Target)
 	{
 		PCHUsage = ModuleRules.PCHUsageMode.UseExplicitOrSharedPCHs;
-        OptimizeCode = CodeOptimization.Never;
 
-        PublicDependencyModuleNames.AddRange(
+		PublicDependencyModuleNames.AddRange(
 			new string[]
 			{
 				"Core"
