[Soundfields] changed check() for IsInGameThread to IsInAudioThread - Hitting check when initting soundfield submix even though the UObjects being protected are safe to access from the Audio Thread.

diff --git a/Engine/Plugins/Runtime/SoundFields/Source/SoundFields/Private/SoundFields.cpp b/Engine/Plugins/Runtime/SoundFields/Source/SoundFields/Private/SoundFields.cpp
--- a/Engine/Plugins/Runtime/SoundFields/Source/SoundFields/Private/SoundFields.cpp
+++ b/Engine/Plugins/Runtime/SoundFields/Source/SoundFields/Private/SoundFields.cpp
@@ -283,7 +283,7 @@ UClass* FAmbisonicsSoundfieldFormat::GetCustomEncodingSettingsClass() const
 
 USoundfieldEncodingSettingsBase* FAmbisonicsSoundfieldFormat::GetDefaultEncodingSettings()
 {
-	check(IsInGameThread());
+	check(IsInGameThread() || IsInAudioThread());
 	return GetMutableDefault<UAmbisonicsEncodingSettings>();
 }
 
